Templates - updated Feature Packs for Top Down, Puzzle and VR templates.

#rnx

diff --git a/Templates/TP_Puzzle/Config/config.ini b/Templates/TP_Puzzle/Config/config.ini
--- a/Templates/TP_Puzzle/Config/config.ini
+++ b/Templates/TP_Puzzle/Config/config.ini
@@ -6,12 +6,5 @@ GameName=TP_Puzzle
 
 [AdditionalFilesToAdd]
 +Files=Templates/TP_Puzzle/Content/*.*
-+Files=Templates/TP_Puzzle/Source/TP_Puzzle/TP_Puzzle.h
-+Files=Templates/TP_Puzzle/Source/TP_Puzzle/TP_PuzzleBlock.cpp
-+Files=Templates/TP_Puzzle/Source/TP_Puzzle/TP_PuzzleBlock.h
-+Files=Templates/TP_Puzzle/Source/TP_Puzzle/TP_PuzzleBlockGrid.cpp
-+Files=Templates/TP_Puzzle/Source/TP_Puzzle/TP_PuzzleBlockGrid.h
-+Files=Templates/TP_Puzzle/Source/TP_Puzzle/TP_PuzzleGameMode.cpp
-+Files=Templates/TP_Puzzle/Source/TP_Puzzle/TP_PuzzleGameMode.h
-+Files=Templates/TP_Puzzle/Source/TP_Puzzle/TP_PuzzlePlayerController.cpp
-+Files=Templates/TP_Puzzle/Source/TP_Puzzle/TP_PuzzlePlayerController.h
++Files=Templates/TP_Puzzle/Source/TP_Puzzle/*.h
++Files=Templates/TP_Puzzle/Source/TP_Puzzle/*.cpp
\ No newline at end of file
diff --git a/Templates/TP_TopDown/Config/config.ini b/Templates/TP_TopDown/Config/config.ini
--- a/Templates/TP_TopDown/Config/config.ini
+++ b/Templates/TP_TopDown/Config/config.ini
@@ -6,10 +6,5 @@ GameName=TP_TopDown
 
 [AdditionalFilesToAdd]
 +Files=Templates/TP_TopDown/Content/*.*
-+Files=Templates/TP_TopDown/Source/TP_TopDown/TP_TopDown.h
-+Files=Templates/TP_TopDown/Source/TP_TopDown/TP_TopDownCharacter.cpp
-+Files=Templates/TP_TopDown/Source/TP_TopDown/TP_TopDownCharacter.h
-+Files=Templates/TP_TopDown/Source/TP_TopDown/TP_TopDownGameMode.cpp
-+Files=Templates/TP_TopDown/Source/TP_TopDown/TP_TopDownGameMode.h
-+Files=Templates/TP_TopDown/Source/TP_TopDown/TP_TopDownPlayerController.cpp
-+Files=Templates/TP_TopDown/Source/TP_TopDown/TP_TopDownPlayerController.h
++Files=Templates/TP_TopDown/Source/TP_TopDown/*.h
++Files=Templates/TP_TopDown/Source/TP_TopDown/*.cpp
\ No newline at end of file
